Initial TimeOfDay plugin

diff --git a/Engine/Plugins/Lighting/TimeOfDay/Source/TimeOfDay/Private/TimeOfDayActor.cpp b/Engine/Plugins/Lighting/TimeOfDay/Source/TimeOfDay/Private/TimeOfDayActor.cpp
new file mode 100644
--- /dev/null
+++ b/Engine/Plugins/Lighting/TimeOfDay/Source/TimeOfDay/Private/TimeOfDayActor.cpp
@@ -0,0 +1,123 @@
+// Copyright Epic Games, Inc. All Rights Reserved.
+
+#include "TimeOfDayActor.h"
+#include "Engine/ActorChannel.h"
+#include "Net/UnrealNetwork.h"
+#include "LevelSequenceActor.h"
+#include "MovieSceneTimeHelpers.h"
+
+ATimeOfDayActor::ATimeOfDayActor(const FObjectInitializer& Init)
+: Super(Init)
+, bRunDayCycle(true)
+, DayLength(24, 0, 0, 0, false)
+, TimePerCycle(0, 5, 0, 0, false)
+, InitialTimeOfDay(6, 0, 0, 0, false)
+{
+	USceneComponent* SceneRootComponent = CreateDefaultSubobject<USceneComponent>(USceneComponent::GetDefaultSceneRootVariableName());
+	SetRootComponent(SceneRootComponent);
+
+#if WITH_EDITORONLY_DATA
+	TimeOfDayPreview = FTimecode(6, 0, 0, 0, false);
+#endif
+}
+
+void ATimeOfDayActor::PostInitializeComponents()
+{
+	// Initialize the player with our settings first so that the LevelSequenceActor's
+	// PostInitializeComponents sequence player initialization is skipped.
+	InitializeSequencePlayer();
+	
+	Super::PostInitializeComponents();
+}
+
+void ATimeOfDayActor::BeginPlay()
+{
+	// Override playback settings with our own based on the
+	// day cycle properties.
+	PlaybackSettings = GetPlaybackSettings();
+	
+	Super::BeginPlay();
+
+	// Day cycle playback settings will always play. Pause if RunDayCycle is
+	// off to allow sequence spawnables and settings to be set from initial
+	// time of day.
+	if (!bRunDayCycle)
+	{
+		SequencePlayer->Pause();
+	}
+}
+
+void ATimeOfDayActor::EndPlay(const EEndPlayReason::Type EndPlayReason)
+{
+	Super::EndPlay(EndPlayReason);
+}
+
+void ATimeOfDayActor::RewindForReplay()
+{
+	Super::RewindForReplay();
+}
+
+void ATimeOfDayActor::InitializeSequencePlayer()
+{
+	if (LevelSequenceAsset && GetWorld()->IsGameWorld())
+	{
+		// Level sequence is already loaded. Initialize the player if it's not already initialized with this sequence
+		if (LevelSequenceAsset != SequencePlayer->GetSequence())
+		{
+			SequencePlayer->Initialize(LevelSequenceAsset, GetLevel(), GetPlaybackSettings(), FLevelSequenceCameraSettings());
+		}
+	}
+}
+
+FMovieSceneSequencePlaybackSettings ATimeOfDayActor::GetPlaybackSettings() const
+{
+	FMovieSceneSequencePlaybackSettings Settings;
+	Settings.bAutoPlay = true;
+	Settings.LoopCount.Value = -1; // Loop indefinitely
+	Settings.PlayRate = 1.0f;
+
+	if (!LevelSequenceAsset)
+	{
+		return Settings;
+	}
+
+	// Update the PlayRate and StartOffset from the DayCycle settings.
+	if (const UMovieScene* MovieScene = LevelSequenceAsset->GetMovieScene())
+	{
+		const FFrameRate DisplayRate = MovieScene->GetDisplayRate();
+		const FFrameRate TickResolution = MovieScene->GetTickResolution();
+		const FFrameNumber FramesPerCycle = TimePerCycle.ToFrameNumber(DisplayRate);
+		const FQualifiedFrameTime DayFrameTime(TimePerCycle, DisplayRate);
+		
+		const TRange<FFrameNumber> MoviePlaybackRange = MovieScene->GetPlaybackRange();
+		if (MoviePlaybackRange.GetLowerBound().IsClosed() && MoviePlaybackRange.GetUpperBound().IsClosed() && FramesPerCycle.Value > 0)
+		{
+			const FFrameNumber SrcStartFrame = UE::MovieScene::DiscreteInclusiveLower(MoviePlaybackRange);
+			const FFrameNumber SrcEndFrame   = UE::MovieScene::DiscreteExclusiveUpper(MoviePlaybackRange);
+
+			const FFrameTime EndingTime = ConvertFrameTime(SrcEndFrame, TickResolution, DisplayRate);
+			const FFrameNumber StartingFrame = ConvertFrameTime(SrcStartFrame, TickResolution, DisplayRate).FloorToFrame();
+			const FFrameNumber EndingFrame   = EndingTime.FloorToFrame();
+			
+			const FQualifiedFrameTime MovieFrameTime(FTimecode(0,0,0,EndingFrame.Value - StartingFrame.Value, false), DisplayRate);
+			Settings.PlayRate = MovieFrameTime.AsSeconds() / DayFrameTime.AsSeconds();
+
+			const FFrameNumber FramesDayLength = DayLength.ToFrameNumber(DisplayRate);
+			const FFrameNumber FramesStartTime = InitialTimeOfDay.ToFrameNumber(DisplayRate);
+			if (FramesDayLength.Value > 0)
+			{
+				const FTimecode InitialTimeOfDayCorrected(0, 0, 0, FramesStartTime.Value % FramesDayLength.Value, false);
+				const FQualifiedFrameTime DayLengthTime(DayLength, DisplayRate);
+				const FQualifiedFrameTime StartTime(InitialTimeOfDayCorrected, DisplayRate);
+				const double StartTimeRatio = StartTime.AsSeconds() / DayLengthTime.AsSeconds();
+
+				// StartTime is in seconds in movie sequence time.
+				Settings.StartTime = StartTimeRatio * MovieFrameTime.AsSeconds();
+			}
+		}
+	}
+
+	return Settings;
+}
+
+
diff --git a/Engine/Plugins/Lighting/TimeOfDay/Source/TimeOfDay/Private/TimeOfDayModule.cpp b/Engine/Plugins/Lighting/TimeOfDay/Source/TimeOfDay/Private/TimeOfDayModule.cpp
new file mode 100644
--- /dev/null
+++ b/Engine/Plugins/Lighting/TimeOfDay/Source/TimeOfDay/Private/TimeOfDayModule.cpp
@@ -0,0 +1,21 @@
+// Copyright Epic Games, Inc. All Rights Reserved.
+
+#include "TimeOfDayModule.h"
+#include "Modules/ModuleManager.h"
+
+#define LOCTEXT_NAMESPACE "FTimeOfDayModule"
+
+void FTimeOfDayModule::StartupModule()
+{
+	// This code will execute after your module is loaded into memory; the exact timing is specified in the .uplugin file per-module
+}
+
+void FTimeOfDayModule::ShutdownModule()
+{
+	// This function may be called during shutdown to clean up your module.  For modules that support dynamic reloading,
+	// we call this function before unloading the module.
+}
+
+#undef LOCTEXT_NAMESPACE
+	
+IMPLEMENT_MODULE(FTimeOfDayModule, TimeOfDay)
\ No newline at end of file
diff --git a/Engine/Plugins/Lighting/TimeOfDay/Source/TimeOfDay/Private/TimeOfDayModule.h b/Engine/Plugins/Lighting/TimeOfDay/Source/TimeOfDay/Private/TimeOfDayModule.h
new file mode 100644
--- /dev/null
+++ b/Engine/Plugins/Lighting/TimeOfDay/Source/TimeOfDay/Private/TimeOfDayModule.h
@@ -0,0 +1,15 @@
+// Copyright Epic Games, Inc. All Rights Reserved.
+
+#pragma once
+
+#include "CoreMinimal.h"
+#include "Modules/ModuleInterface.h"
+
+class FTimeOfDayModule : public IModuleInterface
+{
+public:
+
+	/** IModuleInterface implementation */
+	virtual void StartupModule() override;
+	virtual void ShutdownModule() override;
+};
diff --git a/Engine/Plugins/Lighting/TimeOfDay/Source/TimeOfDay/Private/TimeOfDaySubsystem.cpp b/Engine/Plugins/Lighting/TimeOfDay/Source/TimeOfDay/Private/TimeOfDaySubsystem.cpp
new file mode 100644
--- /dev/null
+++ b/Engine/Plugins/Lighting/TimeOfDay/Source/TimeOfDay/Private/TimeOfDaySubsystem.cpp
@@ -0,0 +1,27 @@
+// Copyright Epic Games, Inc. All Rights Reserved.
+
+
+#include "TimeOfDaySubsystem.h"
+#include "EngineUtils.h"
+#include "TimeOfDayActor.h"
+
+bool UTimeOfDaySubsystem::DoesSupportWorldType(EWorldType::Type WorldType) const
+{
+	return WorldType == EWorldType::Game || WorldType == EWorldType::Editor || WorldType == EWorldType::PIE;
+}
+
+void UTimeOfDaySubsystem::Initialize(FSubsystemCollectionBase& Collection)
+{
+	Super::Initialize(Collection);
+}
+
+void UTimeOfDaySubsystem::Deinitialize()
+{
+	Super::Deinitialize();
+}
+
+ATimeOfDayActor* UTimeOfDaySubsystem::GetTimeOfDayActor() const
+{
+	TActorIterator<ATimeOfDayActor> It(GetWorld());
+	return It ? *It : nullptr;
+}
\ No newline at end of file
diff --git a/Engine/Plugins/Lighting/TimeOfDay/Source/TimeOfDay/Public/TimeOfDayActor.h b/Engine/Plugins/Lighting/TimeOfDay/Source/TimeOfDay/Public/TimeOfDayActor.h
new file mode 100644
--- /dev/null
+++ b/Engine/Plugins/Lighting/TimeOfDay/Source/TimeOfDay/Public/TimeOfDayActor.h
@@ -0,0 +1,58 @@
+// Copyright Epic Games, Inc. All Rights Reserved.
+
+#pragma once
+
+#include "CoreMinimal.h"
+#include "LevelSequenceActor.h"
+#include "TimeOfDayActor.generated.h"
+
+class USkyAtmosphereComponent;
+class USkyLightComponent;
+class UDirectionalLightComponent;
+
+UCLASS(prioritizeCategories=(TimeOfDay, RuntimeDayCycle), hideCategories=(Rendering, Physics, HLOD, Activation, Input, Collision, Actor, Lod, Cooking, DataLayers, Replication, WorldPartition))
+class TIMEOFDAY_API ATimeOfDayActor : public ALevelSequenceActor
+{
+	GENERATED_BODY()
+
+public:
+	ATimeOfDayActor(const FObjectInitializer& Init);
+
+protected:
+	//~ Begin AActor interface
+	virtual void PostInitializeComponents() override;
+	virtual void BeginPlay() override;
+	virtual void EndPlay(const EEndPlayReason::Type EndPlayReason) override;
+	virtual void RewindForReplay() override;
+	//~ End AActor interface
+
+	/** Customized TimeOfDay initialization of the SequencePlayer */
+	void InitializeSequencePlayer();
+
+	/** Compute PlaybackSettings from day cycle properties */
+	FMovieSceneSequencePlaybackSettings GetPlaybackSettings() const;
+
+protected:
+#if WITH_EDITORONLY_DATA
+	/** Sets the time of day to preview in the editor. Does not affect the start time at runtime */
+	UPROPERTY(EditAnywhere, Category = TimeOfDay)
+	FTimecode TimeOfDayPreview;
+#endif
+
+	/** Whether or not to run a day cycle. If this is unchecked the day cycle will remain fixed at the time specified by the Initial Time setting */
+	UPROPERTY(EditAnywhere, Category = RuntimeDayCycle)
+	bool bRunDayCycle;
+
+	/** How long a single day cycle is */
+	UPROPERTY(EditAnywhere, Category = RuntimeDayCycle)
+	FTimecode DayLength;
+
+	/** How long does it take for a day cycle to complete in world time. If this is the same value as day duration that means real world time is used */
+	UPROPERTY(EditAnywhere, Category = RuntimeDayCycle)
+	FTimecode TimePerCycle;
+
+	/** The initial time that the day cycle will start at */
+	UPROPERTY(EditAnywhere, Category = RuntimeDayCycle)
+	FTimecode InitialTimeOfDay;
+};
+
diff --git a/Engine/Plugins/Lighting/TimeOfDay/Source/TimeOfDay/Public/TimeOfDaySubsystem.h b/Engine/Plugins/Lighting/TimeOfDay/Source/TimeOfDay/Public/TimeOfDaySubsystem.h
new file mode 100644
--- /dev/null
+++ b/Engine/Plugins/Lighting/TimeOfDay/Source/TimeOfDay/Public/TimeOfDaySubsystem.h
@@ -0,0 +1,20 @@
+// Copyright Epic Games, Inc. All Rights Reserved.
+
+#pragma once
+
+#include "CoreMinimal.h"
+#include "Subsystems/WorldSubsystem.h"
+//#include "TimeOfDaySubsystem.generated.h"
+
+class ATimeOfDayActor;
+
+class UTimeOfDaySubsystem : public UWorldSubsystem
+{
+public:
+	/** UWorldSubsystem interface */
+	virtual bool DoesSupportWorldType(EWorldType::Type WorldType) const override;
+	virtual void Initialize(FSubsystemCollectionBase& Collection) override;
+	virtual void Deinitialize() override;
+
+	TIMEOFDAY_API ATimeOfDayActor* GetTimeOfDayActor() const;
+};
\ No newline at end of file
diff --git a/Engine/Plugins/Lighting/TimeOfDay/Source/TimeOfDay/TimeOfDay.Build.cs b/Engine/Plugins/Lighting/TimeOfDay/Source/TimeOfDay/TimeOfDay.Build.cs
new file mode 100644
--- /dev/null
+++ b/Engine/Plugins/Lighting/TimeOfDay/Source/TimeOfDay/TimeOfDay.Build.cs
@@ -0,0 +1,49 @@
+// Copyright Epic Games, Inc. All Rights Reserved.
+
+using UnrealBuildTool;
+
+public class TimeOfDay : ModuleRules
+{
+	public TimeOfDay(ReadOnlyTargetRules Target) : base(Target)
+	{
+		PCHUsage = ModuleRules.PCHUsageMode.UseExplicitOrSharedPCHs;
+		
+		PublicIncludePaths.AddRange(
+			new string[] {
+			}
+			);
+				
+		
+		PrivateIncludePaths.AddRange(
+			new string[] {
+			}
+			);
+			
+		
+		PublicDependencyModuleNames.AddRange(
+			new string[]
+			{
+				"Core",
+				"MovieScene",
+				"LevelSequence",
+			}
+			);
+			
+		
+		PrivateDependencyModuleNames.AddRange(
+			new string[]
+			{
+				"InputCore",
+				"CoreUObject",
+				"Engine",
+			}
+			);
+		
+		
+		DynamicallyLoadedModuleNames.AddRange(
+			new string[]
+			{
+			}
+			);
+	}
+}
diff --git a/Engine/Plugins/Lighting/TimeOfDay/Source/TimeOfDayEditor/Private/STimeOfDaySettings.cpp b/Engine/Plugins/Lighting/TimeOfDay/Source/TimeOfDayEditor/Private/STimeOfDaySettings.cpp
new file mode 100644
--- /dev/null
+++ b/Engine/Plugins/Lighting/TimeOfDay/Source/TimeOfDayEditor/Private/STimeOfDaySettings.cpp
@@ -0,0 +1,188 @@
+// Copyright Epic Games, Inc. All Rights Reserved.
+
+#include "STimeOfDaySettings.h"
+#include "Widgets/Input/SSegmentedControl.h"
+#include "Widgets/Layout/SWidgetSwitcher.h"
+#include "SSubobjectInstanceEditor.h"
+#include "EngineUtils.h"
+#include "TimeOfDayActor.h"
+#include "TimeOfDaySubsystem.h"
+#include "ISCSEditorUICustomization.h"
+#include "SPositiveActionButton.h"
+
+#define LOCTEXT_NAMESPACE "TimeOfDayEditor"
+
+class FTimeOfDaySCSEditorUICustomization : public ISCSEditorUICustomization
+{
+public:
+	virtual bool HideComponentsFilterBox() const override { return true; }
+	virtual bool HideBlueprintButtons() const override { return true; }
+};
+
+enum class ESettingsSection : uint8
+{
+	Environment,
+	TimeOfDay,
+};
+
+void STimeOfDaySettings::Construct(const FArguments& InArgs)
+{
+	FEditorDelegates::MapChange.AddSP(this, &STimeOfDaySettings::OnMapChanged);
+
+	UpdateTimeOfDayActor();
+
+	TSharedRef<SSubobjectInstanceEditor> SubObjectEditorRef
+		= SNew(SSubobjectInstanceEditor)
+		.ObjectContext(this, &STimeOfDaySettings::GetObjectContext)
+		//.AllowEditing(this, &SActorDetails::GetAllowComponentTreeEditing)
+		.OnSelectionUpdated(this, &STimeOfDaySettings::OnSubobjectEditorTreeViewSelectionChanged);
+		//.OnItemDoubleClicked(this, &SActorDetails::OnSubobjectEditorTreeViewItemDoubleClicked)
+		//.OnObjectReplaced(this, &SActorDetails::OnSubobjectEditorTreeViewObjectReplaced);
+	
+	SubObjectEditorRef->SetUICustomization(MakeShared<FTimeOfDaySCSEditorUICustomization>());
+
+	SubobjectEditor = SubObjectEditorRef;
+
+	FPropertyEditorModule& PropertyEditorModule = FModuleManager::Get().LoadModuleChecked<FPropertyEditorModule>("PropertyEditor");
+
+	FDetailsViewArgs DetailsViewArgs;
+	DetailsViewArgs.bShowObjectLabel = false;
+	DetailsViewArgs.NameAreaSettings = FDetailsViewArgs::HideNameArea;
+
+	TSharedRef<IDetailsView> DetailsView = PropertyEditorModule.CreateDetailView(DetailsViewArgs);
+	ComponentDetailsView = DetailsView;
+
+	ChildSlot
+	[
+		SNew(SVerticalBox)
+		+ SVerticalBox::Slot()
+		.HAlign(HAlign_Left)
+		.VAlign(VAlign_Top)
+		.Padding(10, 4)
+		.AutoHeight()
+		[
+			SNew(SSegmentedControl<ESettingsSection>)
+			.OnValueChanged(this, &STimeOfDaySettings::OnSettingsSectionChanged)
+			+ SSegmentedControl<ESettingsSection>::Slot(ESettingsSection::Environment)
+			.Text(LOCTEXT("EnvironmentSettings", "Environment"))
+			.ToolTip(LOCTEXT("EnvironmentSettings_ToolTip", "Set up the time of day environment"))
+			+ SSegmentedControl<ESettingsSection>::Slot(ESettingsSection::TimeOfDay)
+			.Text(LOCTEXT("TimeOfDaySettings", "Time of Day"))
+			.ToolTip(LOCTEXT("TimeOfDaySettings_ToolTip", "Specify time of day settings"))
+		]
+		+ SVerticalBox::Slot()
+		[
+			SAssignNew(SettingsSwitcher, SWidgetSwitcher)
+			+ SWidgetSwitcher::Slot()
+			[
+				MakeEnvironmentPanel()
+			]
+			+ SWidgetSwitcher::Slot()
+			[
+				MakeEditDaySequencePanel()
+			]
+		]
+	];
+
+	SettingsSwitcher->SetActiveWidgetIndex(0);
+}
+
+void STimeOfDaySettings::OnSettingsSectionChanged(ESettingsSection NewSection)
+{
+	SettingsSwitcher->SetActiveWidgetIndex(NewSection == ESettingsSection::Environment ? 0 : 1);
+}
+
+void STimeOfDaySettings::OnMapChanged(uint32 Flags)
+{
+	if (Flags == MapChangeEventFlags::NewMap)
+	{
+		UpdateTimeOfDayActor();
+	}
+}
+
+UObject* STimeOfDaySettings::GetObjectContext() const
+{
+	return EditorTimeOfDayActor.Get();
+}
+
+void STimeOfDaySettings::UpdateTimeOfDayActor()
+{
+	EditorTimeOfDayActor = nullptr;
+
+	if (GEditor)
+	{
+		UTimeOfDaySubsystem* TODSubsystem = GEditor->GetEditorWorldContext().World()->GetSubsystem<UTimeOfDaySubsystem>();
+		EditorTimeOfDayActor = TODSubsystem->GetTimeOfDayActor();
+	}
+}
+
+FReply STimeOfDaySettings::OnEditDaySequenceClicked()
+{
+	return FReply::Handled();
+}
+
+void STimeOfDaySettings::OnSubobjectEditorTreeViewSelectionChanged(const TArray<TSharedPtr<FSubobjectEditorTreeNode> >& SelectedNodes)
+{
+	TArray<UObject*> Objects;
+	Objects.Reserve(SelectedNodes.Num());
+
+	for (const TSharedPtr<FSubobjectEditorTreeNode>& Node : SelectedNodes)
+	{
+		if (const UObject* Object = Node->GetObject())
+		{
+			Objects.Add(const_cast<UObject*>(Object));
+		}
+	}
+
+	ComponentDetailsView->SetObjects(Objects);
+}
+
+TSharedRef<SWidget> STimeOfDaySettings::MakeEnvironmentPanel()
+{
+	check(SubobjectEditor.IsValid() && ComponentDetailsView.IsValid());
+	return SNew(SVerticalBox)
+		+ SVerticalBox::Slot()
+		.AutoHeight()
+		[
+			SNew(SHorizontalBox)
+			+ SHorizontalBox::Slot()
+			.AutoWidth()
+			.Padding(10, 4, 4, 4)
+			[
+				SubobjectEditor->GetToolButtonsBox().ToSharedRef()
+			]
+			+ SHorizontalBox::Slot()
+			.AutoWidth()
+			.Padding(4)
+			.VAlign(VAlign_Center)
+			[
+				SNew(SPositiveActionButton)
+				.Text(LOCTEXT("EditDaySequence", "Edit Day Sequence"))
+				.OnClicked(this, &STimeOfDaySettings::OnEditDaySequenceClicked)
+			]
+		]
+		+ SVerticalBox::Slot()
+		[
+			SNew(SSplitter)
+			.Orientation(Orient_Vertical)
+			+ SSplitter::Slot()
+			[
+				SubobjectEditor.ToSharedRef()
+			]
+			+ SSplitter::Slot()
+			[
+				ComponentDetailsView.ToSharedRef()
+			]		
+		];
+}
+
+TSharedRef<SWidget> STimeOfDaySettings::MakeEditDaySequencePanel()
+{
+	return SNew(SVerticalBox)
+		+ SVerticalBox::Slot()
+		.AutoHeight();
+}
+
+
+#undef LOCTEXT_NAMESPACE
+
diff --git a/Engine/Plugins/Lighting/TimeOfDay/Source/TimeOfDayEditor/Private/STimeOfDaySettings.h b/Engine/Plugins/Lighting/TimeOfDay/Source/TimeOfDayEditor/Private/STimeOfDaySettings.h
new file mode 100644
--- /dev/null
+++ b/Engine/Plugins/Lighting/TimeOfDay/Source/TimeOfDayEditor/Private/STimeOfDaySettings.h
@@ -0,0 +1,39 @@
+// Copyright Epic Games, Inc. All Rights Reserved.
+
+#pragma once
+
+#include "Widgets/SCompoundWidget.h"
+
+enum class ESettingsSection : uint8;
+
+class SWidgetSwitcher;
+class SSubobjectInstanceEditor;
+class ATimeOfDayActor;
+class IDetailsView;
+class FSubobjectEditorTreeNode;
+
+class STimeOfDaySettings : public SCompoundWidget
+{
+	SLATE_BEGIN_ARGS(STimeOfDaySettings)
+	{}
+	SLATE_END_ARGS()
+
+public:
+	void Construct(const FArguments& InArgs);
+
+private:
+	void OnSettingsSectionChanged(ESettingsSection NewSection);
+	void OnMapChanged(uint32 Flags);
+	UObject* GetObjectContext() const;
+	void UpdateTimeOfDayActor();
+	FReply OnEditDaySequenceClicked();
+	void OnSubobjectEditorTreeViewSelectionChanged(const TArray<TSharedPtr<FSubobjectEditorTreeNode> >& SelectedNodes);
+
+	TSharedRef<SWidget> MakeEnvironmentPanel();
+	TSharedRef<SWidget> MakeEditDaySequencePanel();
+private:
+	TSharedPtr<SWidgetSwitcher> SettingsSwitcher;
+	TSharedPtr<SSubobjectInstanceEditor> SubobjectEditor;
+	TWeakObjectPtr<ATimeOfDayActor> EditorTimeOfDayActor;
+	TSharedPtr<IDetailsView> ComponentDetailsView;
+};
\ No newline at end of file
diff --git a/Engine/Plugins/Lighting/TimeOfDay/Source/TimeOfDayEditor/Private/TimeOfDayActorDetails.cpp b/Engine/Plugins/Lighting/TimeOfDay/Source/TimeOfDayEditor/Private/TimeOfDayActorDetails.cpp
new file mode 100644
--- /dev/null
+++ b/Engine/Plugins/Lighting/TimeOfDay/Source/TimeOfDayEditor/Private/TimeOfDayActorDetails.cpp
@@ -0,0 +1,34 @@
+// Copyright Epic Games, Inc. All Rights Reserved.
+#include "TimeOfDayActorDetails.h"
+#include "DetailLayoutBuilder.h"
+#include "DetailWidgetRow.h"
+#include "LevelSequenceActor.h"
+
+#define LOCTEXT_NAMESPACE "TimeOfDayActorDetails"
+
+TSharedRef<IDetailCustomization> FTimeOfDayActorDetails::MakeInstance()
+{
+	return MakeShared<FTimeOfDayActorDetails>();
+}
+
+void FTimeOfDayActorDetails::CustomizeDetails(IDetailLayoutBuilder& DetailLayout)
+{
+	DetailLayout.HideProperty("DefaultComponents");
+	
+	// Hide level level sequence actor stuff
+	DetailLayout.HideCategory("Playback");
+	DetailLayout.HideCategory("BurnInOptions");
+	DetailLayout.HideCategory("Aspect Ratio");
+	DetailLayout.HideCategory("Cinematic");
+	DetailLayout.HideCategory("InstanceData");
+	DetailLayout.HideCategory("Tick");
+
+	const TSharedPtr<IPropertyHandle> LevelSequenceProp = DetailLayout.GetProperty(GET_MEMBER_NAME_CHECKED(ALevelSequenceActor, LevelSequenceAsset), ALevelSequenceActor::StaticClass());
+	if (IDetailPropertyRow* DetailRow = DetailLayout.EditDefaultProperty(LevelSequenceProp))
+	{
+		DetailRow->DisplayName(LOCTEXT("EnvironmentPreset", "Environment Preset"));
+	}
+}
+
+
+#undef LOCTEXT_NAMESPACE
diff --git a/Engine/Plugins/Lighting/TimeOfDay/Source/TimeOfDayEditor/Private/TimeOfDayActorDetails.h b/Engine/Plugins/Lighting/TimeOfDay/Source/TimeOfDayEditor/Private/TimeOfDayActorDetails.h
new file mode 100644
--- /dev/null
+++ b/Engine/Plugins/Lighting/TimeOfDay/Source/TimeOfDayEditor/Private/TimeOfDayActorDetails.h
@@ -0,0 +1,19 @@
+// Copyright Epic Games, Inc. All Rights Reserved.
+
+#pragma once
+
+#include "CoreMinimal.h"
+#include "IDetailCustomization.h"
+
+class IDetailLayoutBuilder;
+
+class FTimeOfDayActorDetails : public IDetailCustomization
+{
+public:
+	/** Makes a new instance of this detail layout class for a specific detail view requesting it */
+	static TSharedRef<IDetailCustomization> MakeInstance();
+private:
+	/** IDetailCustomization interface */
+	virtual void CustomizeDetails(IDetailLayoutBuilder& DetailLayout) override;
+private:
+};
diff --git a/Engine/Plugins/Lighting/TimeOfDay/Source/TimeOfDayEditor/Private/TimeOfDayEditor.cpp b/Engine/Plugins/Lighting/TimeOfDay/Source/TimeOfDayEditor/Private/TimeOfDayEditor.cpp
new file mode 100644
--- /dev/null
+++ b/Engine/Plugins/Lighting/TimeOfDay/Source/TimeOfDayEditor/Private/TimeOfDayEditor.cpp
@@ -0,0 +1,180 @@
+// Copyright Epic Games, Inc. All Rights Reserved.
+
+#include "TimeOfDayEditor.h"
+#include "TimeOfDayEditorStyle.h"
+#include "TimeOfDayEditorCommands.h"
+#include "TimeOfDayActorDetails.h"
+#include "STimeOfDaySettings.h"
+#include "Framework/Docking/TabManager.h"
+#include "Widgets/Docking/SDockTab.h"
+#include "ToolMenus.h"
+#include "WorkspaceMenuStructure.h"
+#include "WorkspaceMenuStructureModule.h"
+#include "LevelEditor.h"
+#include "PropertyEditorModule.h"
+
+#define LOCTEXT_NAMESPACE "TimeOfDayEditor"
+
+static const FName TimeOfDayTabName("TimeOfDayEditor");
+static const FName LevelEditorModuleName("LevelEditor");
+static const FName PropertyEditorModuleName("PropertyEditor");
+static const FName TimeOfDayActorClassName("TimeOfDayActor");
+
+void FTimeOfDayEditorModule::StartupModule()
+{
+	// This code will execute after your module is loaded into memory; the exact timing is specified in the .uplugin file per-module
+	
+	FTimeOfDayEditorStyle::Initialize();
+	FTimeOfDayEditorCommands::Register();
+	
+	PluginCommands = MakeShareable(new FUICommandList);
+
+	PluginCommands->MapAction(
+		FTimeOfDayEditorCommands::Get().OpenTimeOfDayEditor,
+		FExecuteAction::CreateRaw(this, &FTimeOfDayEditorModule::OnOpenTimeOfDayEditor),
+		FCanExecuteAction());
+	{
+		const IWorkspaceMenuStructure& MenuStructure = WorkspaceMenu::GetMenuStructure();
+		TSharedRef<FWorkspaceItem> LevelEditorGroup = MenuStructure.GetLevelEditorCategory();
+
+		FLevelEditorModule& LevelEditorModule = FModuleManager::LoadModuleChecked<FLevelEditorModule>(LevelEditorModuleName);
+		
+		auto RegisterTabSpawner = [this, LevelEditorGroup]()
+		{
+			FLevelEditorModule& LevelEditorModule = FModuleManager::LoadModuleChecked<FLevelEditorModule>(LevelEditorModuleName);
+			TSharedPtr<FTabManager> LevelEditorTabManager = LevelEditorModule.GetLevelEditorTabManager();
+
+			LevelEditorTabManager->RegisterTabSpawner(TimeOfDayTabName, FOnSpawnTab::CreateRaw(this, &FTimeOfDayEditorModule::CreateTimeOfDayTab))
+				.SetDisplayName(LOCTEXT("TabTitle", "Time of Day"))
+				.SetTooltipText(LOCTEXT("TooltipText", "Opens the Time of Day editor"))
+				.SetGroup(LevelEditorGroup)
+				.SetIcon(FSlateIcon(FTimeOfDayEditorStyle::GetStyleSetName(), "TimeOfDay.OpenTimeOfDayEditor"));
+
+			RegisterMenus();
+		};
+
+		if (TSharedPtr<FTabManager> LevelEditorTabManager = LevelEditorModule.GetLevelEditorTabManager())
+		{
+			RegisterTabSpawner();
+		}
+		else
+		{
+			LevelEditorModule.OnTabManagerChanged().AddLambda(RegisterTabSpawner);
+		}
+	}
+
+	UToolMenus::RegisterStartupCallback(FSimpleMulticastDelegate::FDelegate::CreateRaw(this, &FTimeOfDayEditorModule::RegisterMenus));
+
+	FPropertyEditorModule& PropertyModule = FModuleManager::LoadModuleChecked<FPropertyEditorModule>(PropertyEditorModuleName);
+	PropertyModule.RegisterCustomClassLayout(TimeOfDayActorClassName, FOnGetDetailCustomizationInstance::CreateStatic(&FTimeOfDayActorDetails::MakeInstance));
+
+	RegisterModulePropertySections();
+}
+
+void FTimeOfDayEditorModule::ShutdownModule()
+{
+	if (FSlateApplication::IsInitialized() && FModuleManager::Get().IsModuleLoaded(LevelEditorModuleName))
+	{
+		FLevelEditorModule* LevelEditorModule = FModuleManager::GetModulePtr<FLevelEditorModule>(LevelEditorModuleName);
+		TSharedPtr<FTabManager> LevelEditorTabManager;
+		if (LevelEditorModule)
+		{
+			LevelEditorTabManager = LevelEditorModule->GetLevelEditorTabManager();
+			LevelEditorModule->OnTabManagerChanged().RemoveAll(this);
+		}
+
+		if (LevelEditorTabManager.IsValid())
+		{
+			LevelEditorTabManager->UnregisterTabSpawner(TimeOfDayTabName);
+		}
+	}
+
+	if (FPropertyEditorModule* PropertyModule = FModuleManager::GetModulePtr<FPropertyEditorModule>(PropertyEditorModuleName))
+	{
+		PropertyModule->UnregisterCustomClassLayout(TimeOfDayActorClassName);
+	}
+
+	DeregisterModulePropertySections();
+
+	UToolMenus::UnRegisterStartupCallback(this);
+
+	UToolMenus::UnregisterOwner(this);
+
+	FTimeOfDayEditorStyle::Shutdown();
+
+	FTimeOfDayEditorCommands::Unregister();
+}
+
+void FTimeOfDayEditorModule::OnOpenTimeOfDayEditor()
+{
+	FLevelEditorModule& LevelEditorModule = FModuleManager::LoadModuleChecked<FLevelEditorModule>(LevelEditorModuleName);
+
+	if (TSharedPtr<FTabManager> LevelEditorTabManager = LevelEditorModule.GetLevelEditorTabManager())
+	{
+		LevelEditorTabManager->TryInvokeTab(TimeOfDayTabName);
+	}
+}
+
+TSharedPtr<FPropertySection> FTimeOfDayEditorModule::RegisterPropertySection(FPropertyEditorModule& PropertyModule, FName ClassName, FName SectionName, FText DisplayName)
+{
+	TSharedRef<FPropertySection> PropertySection = PropertyModule.FindOrCreateSection(ClassName, SectionName, DisplayName);
+	RegisteredPropertySections.Add(ClassName, SectionName);
+	return PropertySection;
+}
+
+void FTimeOfDayEditorModule::RegisterMenus()
+{
+	// Owner will be used for cleanup in call to UToolMenus::UnregisterOwner
+	FToolMenuOwnerScoped OwnerScoped(this);
+	{
+		UToolMenu* ToolbarMenu = UToolMenus::Get()->ExtendMenu("LevelEditor.LevelEditorToolBar.ModesToolBar");
+		{
+			FToolMenuSection& Section = ToolbarMenu->AddSection("TimeOfDay", TAttribute<FText>(), FToolMenuInsert("EditorModes", EToolMenuInsertType::After));
+			{
+				FToolMenuEntry& Entry = Section.AddEntry(
+					FToolMenuEntry::InitToolBarButton(FTimeOfDayEditorCommands::Get().OpenTimeOfDayEditor));
+				Entry.SetCommandList(PluginCommands);
+			}
+		}
+	}
+}
+
+TSharedRef<SDockTab> FTimeOfDayEditorModule::CreateTimeOfDayTab(const FSpawnTabArgs& Args)
+{
+	return
+		SNew(SDockTab)
+		[
+			SNew(STimeOfDaySettings)
+		];
+}
+
+void FTimeOfDayEditorModule::RegisterModulePropertySections()
+{
+	FPropertyEditorModule& PropertyModule = FModuleManager::GetModuleChecked<FPropertyEditorModule>(PropertyEditorModuleName);
+	{
+		TSharedPtr<FPropertySection> Section = RegisterPropertySection(PropertyModule, "TimeOfDayActor", "General", LOCTEXT("General", "General"));
+		Section->AddCategory("General");
+		Section->AddCategory("TimeOfDay");
+		Section->AddCategory("RuntimeDayCycle");
+		Section->AddCategory("BindingOverrides");
+	}
+}
+
+void FTimeOfDayEditorModule::DeregisterModulePropertySections()
+{
+	FPropertyEditorModule* PropertyModule = FModuleManager::GetModulePtr<FPropertyEditorModule>(PropertyEditorModuleName);
+	if (!PropertyModule)
+	{
+		return;
+	}
+
+	for (TMultiMap<FName, FName>::TIterator PropertySectionIterator = RegisteredPropertySections.CreateIterator(); PropertySectionIterator; ++PropertySectionIterator)
+	{
+		PropertyModule->RemoveSection(PropertySectionIterator->Key, PropertySectionIterator->Value);
+		PropertySectionIterator.RemoveCurrent();
+	}
+}
+
+#undef LOCTEXT_NAMESPACE
+	
+IMPLEMENT_MODULE(FTimeOfDayEditorModule, TimeOfDayEditor)
\ No newline at end of file
diff --git a/Engine/Plugins/Lighting/TimeOfDay/Source/TimeOfDayEditor/Private/TimeOfDayEditorCommands.cpp b/Engine/Plugins/Lighting/TimeOfDay/Source/TimeOfDayEditor/Private/TimeOfDayEditorCommands.cpp
new file mode 100644
--- /dev/null
+++ b/Engine/Plugins/Lighting/TimeOfDay/Source/TimeOfDayEditor/Private/TimeOfDayEditorCommands.cpp
@@ -0,0 +1,12 @@
+// Copyright Epic Games, Inc. All Rights Reserved.
+
+#include "TimeOfDayEditorCommands.h"
+
+#define LOCTEXT_NAMESPACE "FTimeOfDayEditorModule"
+
+void FTimeOfDayEditorCommands::RegisterCommands()
+{
+	UI_COMMAND(OpenTimeOfDayEditor, "TimeOfDay", "Opens the time Time of Day Editor", EUserInterfaceActionType::Button, FInputChord());
+}
+
+#undef LOCTEXT_NAMESPACE
diff --git a/Engine/Plugins/Lighting/TimeOfDay/Source/TimeOfDayEditor/Private/TimeOfDayEditorStyle.cpp b/Engine/Plugins/Lighting/TimeOfDay/Source/TimeOfDayEditor/Private/TimeOfDayEditorStyle.cpp
new file mode 100644
--- /dev/null
+++ b/Engine/Plugins/Lighting/TimeOfDay/Source/TimeOfDayEditor/Private/TimeOfDayEditorStyle.cpp
@@ -0,0 +1,65 @@
+// Copyright Epic Games, Inc. All Rights Reserved.
+
+#include "TimeOfDayEditorStyle.h"
+#include "TimeOfDayEditor.h"
+#include "Framework/Application/SlateApplication.h"
+#include "Styling/SlateStyleRegistry.h"
+#include "Slate/SlateGameResources.h"
+#include "Interfaces/IPluginManager.h"
+#include "Styling/SlateStyleMacros.h"
+
+// This is to fix the issue that SlateStyleMacros like IMAGE_BRUSH look for RootToContentDir but StyleSet->RootToContentDir is how this style is set up
+#define RootToContentDir Style->RootToContentDir
+
+TSharedPtr<FSlateStyleSet> FTimeOfDayEditorStyle::StyleInstance = nullptr;
+
+void FTimeOfDayEditorStyle::Initialize()
+{
+	if (!StyleInstance.IsValid())
+	{
+		StyleInstance = Create();
+		FSlateStyleRegistry::RegisterSlateStyle(*StyleInstance);
+	}
+}
+
+void FTimeOfDayEditorStyle::Shutdown()
+{
+	FSlateStyleRegistry::UnRegisterSlateStyle(*StyleInstance);
+	ensure(StyleInstance.IsUnique());
+	StyleInstance.Reset();
+}
+
+FName FTimeOfDayEditorStyle::GetStyleSetName()
+{
+	static FName StyleSetName(TEXT("TimeOfDayStyle"));
+	return StyleSetName;
+}
+
+
+const FVector2D Icon16x16(16.0f, 16.0f);
+const FVector2D Icon20x20(20.0f, 20.0f);
+const FVector2D Icon40x40(40.0f, 40.0f);
+
+TSharedRef< FSlateStyleSet > FTimeOfDayEditorStyle::Create()
+{
+	TSharedRef<FSlateStyleSet> Style = MakeShareable(new FSlateStyleSet("TimeOfDayStyle"));
+	Style->SetContentRoot(IPluginManager::Get().FindPlugin("TimeOfDay")->GetBaseDir() / TEXT("Resources"));
+
+	Style->Set("TimeOfDay.OpenTimeOfDayEditor", new IMAGE_BRUSH_SVG("LightBulb", Icon20x20));
+
+	return Style;
+}
+
+
+void FTimeOfDayEditorStyle::ReloadTextures()
+{
+	if (FSlateApplication::IsInitialized())
+	{
+		FSlateApplication::Get().GetRenderer()->ReloadTextureResources();
+	}
+}
+
+const ISlateStyle& FTimeOfDayEditorStyle::Get()
+{
+	return *StyleInstance;
+}
diff --git a/Engine/Plugins/Lighting/TimeOfDay/Source/TimeOfDayEditor/Public/TimeOfDayEditor.h b/Engine/Plugins/Lighting/TimeOfDay/Source/TimeOfDayEditor/Public/TimeOfDayEditor.h
new file mode 100644
--- /dev/null
+++ b/Engine/Plugins/Lighting/TimeOfDay/Source/TimeOfDayEditor/Public/TimeOfDayEditor.h
@@ -0,0 +1,37 @@
+// Copyright Epic Games, Inc. All Rights Reserved.
+
+#pragma once
+
+#include "CoreMinimal.h"
+#include "Modules/ModuleManager.h"
+
+class FToolBarBuilder;
+class FMenuBuilder;
+class SDockTab;
+class FSpawnTabArgs;
+class FPropertySection;
+class FPropertyEditorModule;
+
+class FTimeOfDayEditorModule : public IModuleInterface
+{
+public:
+
+	/** IModuleInterface implementation */
+	virtual void StartupModule() override;
+	virtual void ShutdownModule() override;
+	
+	void OnOpenTimeOfDayEditor();
+
+	TSharedPtr<FPropertySection> RegisterPropertySection(FPropertyEditorModule& PropertyModule, FName ClassName, FName SectionName, FText DisplayName);
+	void RegisterModulePropertySections();
+	void DeregisterModulePropertySections();
+
+private:
+	void RegisterMenus();
+
+	TSharedRef<SDockTab> CreateTimeOfDayTab(const FSpawnTabArgs& Args);
+
+	TSharedPtr<class FUICommandList> PluginCommands;
+
+	TMultiMap<FName, FName> RegisteredPropertySections;
+};
diff --git a/Engine/Plugins/Lighting/TimeOfDay/Source/TimeOfDayEditor/Public/TimeOfDayEditorCommands.h b/Engine/Plugins/Lighting/TimeOfDay/Source/TimeOfDayEditor/Public/TimeOfDayEditorCommands.h
new file mode 100644
--- /dev/null
+++ b/Engine/Plugins/Lighting/TimeOfDay/Source/TimeOfDayEditor/Public/TimeOfDayEditorCommands.h
@@ -0,0 +1,23 @@
+// Copyright Epic Games, Inc. All Rights Reserved.
+
+#pragma once
+
+#include "CoreMinimal.h"
+#include "Framework/Commands/Commands.h"
+#include "TimeOfDayEditorStyle.h"
+
+class FTimeOfDayEditorCommands : public TCommands<FTimeOfDayEditorCommands>
+{
+public:
+
+	FTimeOfDayEditorCommands()
+		: TCommands<FTimeOfDayEditorCommands>(TEXT("TimeOfDay"), NSLOCTEXT("Contexts", "TimeOfDay", "TimeOfDay Plugin"), NAME_None, FTimeOfDayEditorStyle::GetStyleSetName())
+	{
+	}
+
+	// TCommands<> interface
+	virtual void RegisterCommands() override;
+
+public:
+	TSharedPtr<FUICommandInfo> OpenTimeOfDayEditor;
+};
diff --git a/Engine/Plugins/Lighting/TimeOfDay/Source/TimeOfDayEditor/Public/TimeOfDayEditorStyle.h b/Engine/Plugins/Lighting/TimeOfDay/Source/TimeOfDayEditor/Public/TimeOfDayEditorStyle.h
new file mode 100644
--- /dev/null
+++ b/Engine/Plugins/Lighting/TimeOfDay/Source/TimeOfDayEditor/Public/TimeOfDayEditorStyle.h
@@ -0,0 +1,31 @@
+// Copyright Epic Games, Inc. All Rights Reserved.
+
+#pragma once
+
+#include "CoreMinimal.h"
+#include "Styling/SlateStyle.h"
+
+class FTimeOfDayEditorStyle
+{
+public:
+
+	static void Initialize();
+
+	static void Shutdown();
+
+	/** reloads textures used by slate renderer */
+	static void ReloadTextures();
+
+	/** @return The Slate style set for the Shooter game */
+	static const ISlateStyle& Get();
+
+	static FName GetStyleSetName();
+
+private:
+
+	static TSharedRef< class FSlateStyleSet > Create();
+
+private:
+
+	static TSharedPtr< class FSlateStyleSet > StyleInstance;
+};
\ No newline at end of file
diff --git a/Engine/Plugins/Lighting/TimeOfDay/Source/TimeOfDayEditor/TimeOfDayEditor.Build.cs b/Engine/Plugins/Lighting/TimeOfDay/Source/TimeOfDayEditor/TimeOfDayEditor.Build.cs
new file mode 100644
--- /dev/null
+++ b/Engine/Plugins/Lighting/TimeOfDay/Source/TimeOfDayEditor/TimeOfDayEditor.Build.cs
@@ -0,0 +1,65 @@
+// Copyright Epic Games, Inc. All Rights Reserved.
+
+using UnrealBuildTool;
+
+public class TimeOfDayEditor : ModuleRules
+{
+	public TimeOfDayEditor(ReadOnlyTargetRules Target) : base(Target)
+	{
+		PCHUsage = ModuleRules.PCHUsageMode.UseExplicitOrSharedPCHs;
+		
+		PublicIncludePaths.AddRange(
+			new string[] {
+				// ... add public include paths required here ...
+			}
+			);
+				
+		
+		PrivateIncludePaths.AddRange(
+			new string[] {
+				// ... add other private include paths required here ...
+			}
+			);
+			
+		
+		PublicDependencyModuleNames.AddRange(
+			new string[]
+			{
+				"Core",
+				// ... add other public dependencies that you statically link with here ...
+			}
+			);
+			
+		
+		PrivateDependencyModuleNames.AddRange(
+			new string[]
+			{
+				"Projects",
+				"InputCore",
+				"EditorFramework",
+				"UnrealEd",
+				"EditorWidgets",
+				"ToolMenus",
+				"CoreUObject",
+				"Engine",
+				"Slate",
+				"SlateCore",
+				"WorkspaceMenuStructure",
+				"SubobjectEditor",
+				"TimeOfDay",
+				"ToolWidgets",
+				"LevelSequence"
+				// ... add private dependencies that you statically link with here ...	
+			}
+			);
+		
+		
+		DynamicallyLoadedModuleNames.AddRange(
+			new string[]
+			{
+				// ... add any modules that your module loads dynamically here ...
+				"PropertyEditor"
+			}
+			);
+	}
+}
diff --git a/Engine/Plugins/Lighting/TimeOfDay/TimeOfDay.uplugin b/Engine/Plugins/Lighting/TimeOfDay/TimeOfDay.uplugin
new file mode 100644
--- /dev/null
+++ b/Engine/Plugins/Lighting/TimeOfDay/TimeOfDay.uplugin
@@ -0,0 +1,35 @@
+{
+	"FileVersion": 3,
+	"Version": 1,
+	"VersionName": "1.0",
+	"FriendlyName": "TimeOfDay",
+	"Description": "",
+	"Category": "Other",
+	"CreatedBy": "Epic Games, Inc.",
+	"CreatedByURL": "http://epicgames.com",
+	"DocsURL": "",
+	"MarketplaceURL": "",
+	"SupportURL": "",
+	"CanContainContent": true,
+	"IsBetaVersion": false,
+	"IsExperimentalVersion": true,
+	"Installed": false,
+	"Modules": [
+		{
+			"Name": "TimeOfDayEditor",
+			"Type": "Editor",
+			"LoadingPhase": "Default"
+		},
+		{
+			"Name": "TimeOfDay",
+			"Type": "Runtime",
+			"LoadingPhase": "Default"
+		}
+	],
+    "Plugins": [
+		{
+			"Name": "LevelSequenceEditor",
+			"Enabled": true
+		}
+	]
+}
