TimeOfDay: Load environment preset as a subsequence to a master sequence.

diff --git a/Engine/Plugins/Lighting/TimeOfDay/Source/TimeOfDay/Private/TimeOfDayActor.cpp b/Engine/Plugins/Lighting/TimeOfDay/Source/TimeOfDay/Private/TimeOfDayActor.cpp
--- a/Engine/Plugins/Lighting/TimeOfDay/Source/TimeOfDay/Private/TimeOfDayActor.cpp
+++ b/Engine/Plugins/Lighting/TimeOfDay/Source/TimeOfDay/Private/TimeOfDayActor.cpp
@@ -5,6 +5,9 @@
 #include "Net/UnrealNetwork.h"
 #include "LevelSequenceActor.h"
 #include "MovieSceneTimeHelpers.h"
+#include "Tracks/MovieSceneSubTrack.h"
+#include "Sections/MovieSceneSubSection.h"
+#include "EngineUtils.h"
 
 ATimeOfDayActor::ATimeOfDayActor(const FObjectInitializer& Init)
 : Super(Init)
@@ -23,18 +26,28 @@ ATimeOfDayActor::ATimeOfDayActor(const FObjectInitializer& Init)
 
 void ATimeOfDayActor::PostInitializeComponents()
 {
-	// Initialize the player with our settings first so that the LevelSequenceActor's
-	// PostInitializeComponents sequence player initialization is skipped.
-	InitializeSequencePlayer();
+	// Skip ALevelSequenceActor::PostInitializeComponents to override
+	// the sequence player initialization.
+	AActor::PostInitializeComponents();
+
+	if (HasAuthority())
+	{
+		SetReplicates(bReplicatePlayback);
+	}
 	
-	Super::PostInitializeComponents();
+	// Initialize this player for tick as soon as possible to ensure that a persistent
+	// reference to the tick manager is maintained
+	SequencePlayer->InitializeForTick(this);
+
+	InitializeSequencePlayer();
 }
 
+
 void ATimeOfDayActor::BeginPlay()
 {
 	// Override playback settings with our own based on the
 	// day cycle properties.
-	PlaybackSettings = GetPlaybackSettings();
+	PlaybackSettings = GetPlaybackSettings(MasterSequence);
 	
 	Super::BeginPlay();
 
@@ -59,30 +72,50 @@ void ATimeOfDayActor::RewindForReplay()
 
 void ATimeOfDayActor::InitializeSequencePlayer()
 {
+	MasterSequence = NewObject<ULevelSequence>(this, NAME_None, RF_Transactional);
+	MasterSequence->Initialize();
+	MasterSequence->SetSequenceFlags(EMovieSceneSequenceFlags::Volatile);
+
+	UMovieScene* MasterMovieScene = MasterSequence->GetMovieScene();
+	const int32 MasterDuration = MasterMovieScene->GetTickResolution().AsFrameNumber(1).Value;
+	MasterMovieScene->SetPlaybackRange(0, MasterDuration);
+	MasterMovieScene->AddMasterTrack<UMovieSceneSubTrack>();
+
+	PlaybackSettings = GetPlaybackSettings(MasterSequence);
+	SequencePlayer->Initialize(MasterSequence, GetLevel(), PlaybackSettings, FLevelSequenceCameraSettings());
+
 	if (LevelSequenceAsset && GetWorld()->IsGameWorld())
 	{
-		// Level sequence is already loaded. Initialize the player if it's not already initialized with this sequence
-		if (LevelSequenceAsset != SequencePlayer->GetSequence())
-		{
-			SequencePlayer->Initialize(LevelSequenceAsset, GetLevel(), GetPlaybackSettings(), FLevelSequenceCameraSettings());
-		}
+		UMovieSceneSubTrack* SubTrack = MasterMovieScene->FindMasterTrack<UMovieSceneSubTrack>();
+
+		// Compute outer duration from subsequence asset.
+		const FFrameRate TickResolution = LevelSequenceAsset->GetMovieScene()->GetTickResolution();
+		const FQualifiedFrameTime InnerDuration = FQualifiedFrameTime(
+			UE::MovieScene::DiscreteSize(LevelSequenceAsset->GetMovieScene()->GetPlaybackRange()),
+			TickResolution);
+
+		const FFrameRate OuterFrameRate = SubTrack->GetTypedOuter<UMovieScene>()->GetTickResolution();
+		const int32      OuterDuration  = InnerDuration.ConvertTo(OuterFrameRate).FrameNumber.Value;
+		
+		UMovieSceneSubSection* SubSection = SubTrack->AddSequence(LevelSequenceAsset, 0, OuterDuration);
+		SubSection->Parameters.TimeScale = (float)OuterDuration / (float)MasterDuration;
 	}
 }
 
-FMovieSceneSequencePlaybackSettings ATimeOfDayActor::GetPlaybackSettings() const
+FMovieSceneSequencePlaybackSettings ATimeOfDayActor::GetPlaybackSettings(const ULevelSequence* Sequence) const
 {
 	FMovieSceneSequencePlaybackSettings Settings;
 	Settings.bAutoPlay = true;
 	Settings.LoopCount.Value = -1; // Loop indefinitely
 	Settings.PlayRate = 1.0f;
 
-	if (!LevelSequenceAsset)
+	if (!Sequence)
 	{
 		return Settings;
 	}
 
 	// Update the PlayRate and StartOffset from the DayCycle settings.
-	if (const UMovieScene* MovieScene = LevelSequenceAsset->GetMovieScene())
+	if (const UMovieScene* MovieScene = Sequence->GetMovieScene())
 	{
 		const FFrameRate DisplayRate = MovieScene->GetDisplayRate();
 		const FFrameRate TickResolution = MovieScene->GetTickResolution();
diff --git a/Engine/Plugins/Lighting/TimeOfDay/Source/TimeOfDay/Public/TimeOfDayActor.h b/Engine/Plugins/Lighting/TimeOfDay/Source/TimeOfDay/Public/TimeOfDayActor.h
--- a/Engine/Plugins/Lighting/TimeOfDay/Source/TimeOfDay/Public/TimeOfDayActor.h
+++ b/Engine/Plugins/Lighting/TimeOfDay/Source/TimeOfDay/Public/TimeOfDayActor.h
@@ -30,9 +30,12 @@ protected:
 	void InitializeSequencePlayer();
 
 	/** Compute PlaybackSettings from day cycle properties */
-	FMovieSceneSequencePlaybackSettings GetPlaybackSettings() const;
+	FMovieSceneSequencePlaybackSettings GetPlaybackSettings(const ULevelSequence* Sequence) const;
 
 protected:
+	UPROPERTY()
+	TObjectPtr<ULevelSequence> MasterSequence;
+
 #if WITH_EDITORONLY_DATA
 	/** Sets the time of day to preview in the editor. Does not affect the start time at runtime */
 	UPROPERTY(EditAnywhere, Category = TimeOfDay)
