TimeOfDay: Refactor ATimeOfDayActor to no longer derive from ALevelSequenceActor.

diff --git a/Engine/Plugins/Lighting/TimeOfDay/Source/TimeOfDay/Private/TimeOfDayActor.cpp b/Engine/Plugins/Lighting/TimeOfDay/Source/TimeOfDay/Private/TimeOfDayActor.cpp
--- a/Engine/Plugins/Lighting/TimeOfDay/Source/TimeOfDay/Private/TimeOfDayActor.cpp
+++ b/Engine/Plugins/Lighting/TimeOfDay/Source/TimeOfDay/Private/TimeOfDayActor.cpp
@@ -1,13 +1,16 @@
 // Copyright Epic Games, Inc. All Rights Reserved.
 
 #include "TimeOfDayActor.h"
+#include "Engine/Texture2D.h"
+#include "Components/BillboardComponent.h"
 #include "Engine/ActorChannel.h"
-#include "Net/UnrealNetwork.h"
-#include "LevelSequenceActor.h"
 #include "MovieSceneTimeHelpers.h"
 #include "Tracks/MovieSceneSubTrack.h"
 #include "Sections/MovieSceneSubSection.h"
 #include "EngineUtils.h"
+#include "Net/UnrealNetwork.h"
+#include "LevelSequenceActor.h"
+
 
 ATimeOfDayActor::ATimeOfDayActor(const FObjectInitializer& Init)
 : Super(Init)
@@ -20,15 +23,49 @@ ATimeOfDayActor::ATimeOfDayActor(const FObjectInitializer& Init)
 	SetRootComponent(SceneRootComponent);
 
 #if WITH_EDITORONLY_DATA
+	UBillboardComponent* SpriteComponent = CreateEditorOnlyDefaultSubobject<UBillboardComponent>(TEXT("Sprite"));
+
+	if (!IsRunningCommandlet())
+	{
+		// Structure to hold one-time initialization
+		struct FConstructorStatics
+		{
+			ConstructorHelpers::FObjectFinderOptional<UTexture2D> DecalTexture;
+			FConstructorStatics() : DecalTexture(TEXT("/Engine/EditorResources/S_LevelSequence")) {}
+		};
+		static FConstructorStatics ConstructorStatics;
+
+		if (SpriteComponent)
+		{
+			SpriteComponent->Sprite = ConstructorStatics.DecalTexture.Get();
+			SpriteComponent->SetupAttachment(RootComponent);
+			SpriteComponent->bIsScreenSizeScaled = true;
+			SpriteComponent->SetUsingAbsoluteScale(true);
+			SpriteComponent->bReceivesDecals = false;
+			SpriteComponent->bHiddenInGame = true;
+		}
+	}
+
+	bIsSpatiallyLoaded = false;
+
 	TimeOfDayPreview = FTimecode(6, 0, 0, 0, false);
-#endif
+#endif // WITH_EDITORONLY_DATA
+
+	// The TimeOfDayActor defaults to never ticking by the tick manager because it is ticked separately in LevelTick
+	//PrimaryActorTick.bCanEverTick = false;
+
+	// SequencePlayer must be a default sub object for it to be replicated correctly
+	SequencePlayer = Init.CreateDefaultSubobject<ULevelSequencePlayer>(this, "AnimationPlayer");
+	BindingOverrides = Init.CreateDefaultSubobject<UMovieSceneBindingOverrides>(this, "BindingOverrides");
+
+	bReplicates = true;
+	bReplicatePlayback = false;
+	bReplicateUsingRegisteredSubObjectList = false;
 }
 
 void ATimeOfDayActor::PostInitializeComponents()
 {
-	// Skip ALevelSequenceActor::PostInitializeComponents to override
-	// the sequence player initialization.
-	AActor::PostInitializeComponents();
+	Super::PostInitializeComponents();
 
 	if (HasAuthority())
 	{
@@ -39,21 +76,87 @@ void ATimeOfDayActor::PostInitializeComponents()
 	// reference to the tick manager is maintained
 	SequencePlayer->InitializeForTick(this);
 
-	InitializeSequencePlayer();
+	InitializePlayer();
+}
+
+ULevelSequencePlayer* ATimeOfDayActor::GetSequencePlayer() const
+{
+	return SequencePlayer && SequencePlayer->GetSequence() ? SequencePlayer : nullptr;
+}
+
+ULevelSequence* ATimeOfDayActor::GetDaySequence() const
+{
+	return DaySequenceAsset;
+}
+
+void ATimeOfDayActor::SetDaySequence(ULevelSequence* InSequence)
+{
+	if (!SequencePlayer->IsPlaying())
+	{
+		UpdateDaySequence(InSequence);
+	}
+}
+
+void ATimeOfDayActor::SetReplicatePlayback(bool bInReplicatePlayback)
+{
+	bReplicatePlayback = bInReplicatePlayback;
+	SetReplicates(bReplicatePlayback);
+}
+
+bool ATimeOfDayActor::ReplicateSubobjects(UActorChannel* Channel, FOutBunch* Bunch, FReplicationFlags* RepFlags)
+{
+	bool bWroteSomething = Super::ReplicateSubobjects(Channel, Bunch, RepFlags);
+
+	bWroteSomething |= Channel->ReplicateSubobject(SequencePlayer, *Bunch, *RepFlags);
+
+	return bWroteSomething;
+}
+
+void ATimeOfDayActor::GetLifetimeReplicatedProps(TArray<FLifetimeProperty>& OutLifetimeProps) const
+{
+	Super::GetLifetimeReplicatedProps(OutLifetimeProps);
+
+	DOREPLIFETIME(ATimeOfDayActor, SequencePlayer);
+}
+
+void ATimeOfDayActor::PostInitProperties()
+{
+	Super::PostInitProperties();
+
+	// Have to initialize this here as any properties set on default subobjects inside the constructor
+	// get stomped by the CDO's properties when the constructor exits.
+	SequencePlayer->SetPlaybackClient(this);
 }
 
+void ATimeOfDayActor::PostLoad()
+{
+	Super::PostLoad();
+
+#if WITH_EDITORONLY_DATA
+	// Fix sprite component so that it's attached to the root component. In the past, the sprite component was the root component.
+	UBillboardComponent* SpriteComponent = FindComponentByClass<UBillboardComponent>();
+	if (SpriteComponent && SpriteComponent->GetAttachParent() != RootComponent)
+	{
+		SpriteComponent->SetupAttachment(RootComponent);
+	}
+#endif
+}
 
 void ATimeOfDayActor::BeginPlay()
 {
-	// Override playback settings with our own based on the
-	// day cycle properties.
-	PlaybackSettings = GetPlaybackSettings(MasterSequence);
+	UMovieSceneSequenceTickManager* TickManager = SequencePlayer->GetTickManager();
+	if (ensure(TickManager))
+	{
+		TickManager->RegisterSequenceActor(this);
+	}
 	
 	Super::BeginPlay();
 
 	// Day cycle playback settings will always play. Pause if RunDayCycle is
 	// off to allow sequence spawnables and settings to be set from initial
 	// time of day.
+	SequencePlayer->Play();
+	
 	if (!bRunDayCycle)
 	{
 		SequencePlayer->Pause();
@@ -70,7 +173,56 @@ void ATimeOfDayActor::RewindForReplay()
 	Super::RewindForReplay();
 }
 
-void ATimeOfDayActor::InitializeSequencePlayer()
+void ATimeOfDayActor::TickFromSequenceTickManager(float DeltaSeconds)
+{
+	if (SequencePlayer)
+	{
+		SequencePlayer->UpdateAsync(DeltaSeconds);
+	}
+}
+
+bool ATimeOfDayActor::RetrieveBindingOverrides(const FGuid& InBindingId, FMovieSceneSequenceID InSequenceID, TArray<UObject*, TInlineAllocator<1>>& OutObjects) const
+{
+	// Specify MovieSceneSequenceID::Root while the sequence is always the top level master sequence.
+	return BindingOverrides->LocateBoundObjects(InBindingId, MovieSceneSequenceID::Root, OutObjects);
+}
+
+UObject* ATimeOfDayActor::GetInstanceData() const
+{
+	return nullptr;
+}
+
+bool ATimeOfDayActor::GetIsReplicatedPlayback() const
+{
+	return bReplicatePlayback;
+}
+
+#if WITH_EDITOR
+TSharedPtr<FStructOnScope> ATimeOfDayActor::GetObjectPickerProxy(TSharedPtr<IPropertyHandle> ObjectPropertyHandle)
+{
+	TSharedRef<FStructOnScope> Struct = MakeShared<FStructOnScope>(FBoundActorProxy::StaticStruct());
+	reinterpret_cast<FBoundActorProxy*>(Struct->GetStructMemory())->Initialize(ObjectPropertyHandle);
+	return Struct;
+}
+
+void ATimeOfDayActor::UpdateObjectFromProxy(FStructOnScope& Proxy, IPropertyHandle& ObjectPropertyHandle)
+{
+	UObject* BoundActor = reinterpret_cast<FBoundActorProxy*>(Proxy.GetStructMemory())->BoundActor;
+	ObjectPropertyHandle.SetValue(BoundActor);
+}
+
+bool ATimeOfDayActor::GetReferencedContentObjects(TArray<UObject*>& Objects) const
+{
+	if (DaySequenceAsset)
+	{
+		Objects.Add(DaySequenceAsset);
+	}
+	Super::GetReferencedContentObjects(Objects);
+	return true;
+}
+#endif
+
+void ATimeOfDayActor::InitializePlayer()
 {
 	MasterSequence = NewObject<ULevelSequence>(this, NAME_None, RF_Transactional);
 	MasterSequence->Initialize();
@@ -81,23 +233,47 @@ void ATimeOfDayActor::InitializeSequencePlayer()
 	MasterMovieScene->SetPlaybackRange(0, MasterDuration);
 	MasterMovieScene->AddMasterTrack<UMovieSceneSubTrack>();
 
-	PlaybackSettings = GetPlaybackSettings(MasterSequence);
-	SequencePlayer->Initialize(MasterSequence, GetLevel(), PlaybackSettings, FLevelSequenceCameraSettings());
+	SequencePlayer->Initialize(MasterSequence, GetLevel(), GetPlaybackSettings(MasterSequence), FLevelSequenceCameraSettings());
 
-	if (LevelSequenceAsset && GetWorld()->IsGameWorld())
+	UpdateDaySequence(DaySequenceAsset);
+}
+
+void ATimeOfDayActor::UpdateDaySequence(ULevelSequence* SequenceAsset)
+{
+	if (!MasterSequence)
+	{
+		return;
+	}
+
+	UMovieScene* MasterMovieScene = MasterSequence->GetMovieScene();
+	const int32 MasterDuration = MasterMovieScene->GetPlaybackRange().GetUpperBoundValue().Value;
+	
+	if (SequenceAsset && GetWorld()->IsGameWorld())
 	{
 		UMovieSceneSubTrack* SubTrack = MasterMovieScene->FindMasterTrack<UMovieSceneSubTrack>();
 
+		for (UMovieSceneSection* Section : SubTrack->GetAllSections())
+		{
+			UMovieSceneSubSection* SubSection = Cast<UMovieSceneSubSection>(Section);
+			if (SubSection && SubSection->GetSequence() == DaySequenceAsset)
+			{
+				SubTrack->RemoveSection(*SubSection);
+				break;
+			}
+		}
+
+		DaySequenceAsset = SequenceAsset;
+
 		// Compute outer duration from subsequence asset.
-		const FFrameRate TickResolution = LevelSequenceAsset->GetMovieScene()->GetTickResolution();
+		const FFrameRate TickResolution = DaySequenceAsset->GetMovieScene()->GetTickResolution();
 		const FQualifiedFrameTime InnerDuration = FQualifiedFrameTime(
-			UE::MovieScene::DiscreteSize(LevelSequenceAsset->GetMovieScene()->GetPlaybackRange()),
+			UE::MovieScene::DiscreteSize(DaySequenceAsset->GetMovieScene()->GetPlaybackRange()),
 			TickResolution);
 
 		const FFrameRate OuterFrameRate = SubTrack->GetTypedOuter<UMovieScene>()->GetTickResolution();
 		const int32      OuterDuration  = InnerDuration.ConvertTo(OuterFrameRate).FrameNumber.Value;
 		
-		UMovieSceneSubSection* SubSection = SubTrack->AddSequence(LevelSequenceAsset, 0, OuterDuration);
+		UMovieSceneSubSection* SubSection = SubTrack->AddSequence(DaySequenceAsset, 0, OuterDuration);
 		SubSection->Parameters.TimeScale = (float)OuterDuration / (float)MasterDuration;
 	}
 }
diff --git a/Engine/Plugins/Lighting/TimeOfDay/Source/TimeOfDay/Public/TimeOfDayActor.h b/Engine/Plugins/Lighting/TimeOfDay/Source/TimeOfDay/Public/TimeOfDayActor.h
--- a/Engine/Plugins/Lighting/TimeOfDay/Source/TimeOfDay/Public/TimeOfDayActor.h
+++ b/Engine/Plugins/Lighting/TimeOfDay/Source/TimeOfDay/Public/TimeOfDayActor.h
@@ -3,37 +3,118 @@
 #pragma once
 
 #include "CoreMinimal.h"
-#include "LevelSequenceActor.h"
+#include "GameFramework/Actor.h"
+#include "LevelSequencePlayer.h"
+#include "MovieSceneBindingOwnerInterface.h"
+#include "MovieSceneBindingOverrides.h"
+#include "MovieSceneSequenceTickManager.h"
 #include "TimeOfDayActor.generated.h"
 
-class USkyAtmosphereComponent;
-class USkyLightComponent;
-class UDirectionalLightComponent;
-
-UCLASS(prioritizeCategories=(TimeOfDay, RuntimeDayCycle), hideCategories=(Rendering, Physics, HLOD, Activation, Input, Collision, Actor, Lod, Cooking, DataLayers, Replication, WorldPartition))
-class TIMEOFDAY_API ATimeOfDayActor : public ALevelSequenceActor
+UCLASS(prioritizeCategories=(General, TimeOfDay, RuntimeDayCycle, Binding), hideCategories=(Rendering, Physics, HLOD, Activation, Input, Collision, Actor, Lod, Cooking, DataLayers, Replication, WorldPartition))
+class TIMEOFDAY_API ATimeOfDayActor
+	: public AActor
+	, public IMovieSceneSequenceActor
+	, public IMovieScenePlaybackClient
+	, public IMovieSceneBindingOwnerInterface
 {
 	GENERATED_BODY()
 
 public:
 	ATimeOfDayActor(const FObjectInitializer& Init);
 
+	/** Access this actor's sequence player, or None if it is not yet initialized */
+	UFUNCTION(BlueprintGetter)
+	ULevelSequencePlayer* GetSequencePlayer() const;
+
+	/**
+	 * Get the level sequence being played by this actor.
+	 *
+	 * @return Level sequence, or nullptr if not assigned or if it cannot be loaded.
+	 * @see SetSequence
+	 */
+	UFUNCTION(BlueprintCallable, Category="Sequencer|Player")
+	ULevelSequence* GetDaySequence() const;
+
+	/**
+	 * Set the level sequence being played by this actor.
+	 *
+	 * @param InSequence The sequence object to set.
+	 * @see GetSequence
+	 */
+	UFUNCTION(BlueprintCallable, Category="Sequencer|Player")
+	void SetDaySequence(ULevelSequence* InSequence);
+
+	/** Set whether or not to replicate playback for this actor */
+	UFUNCTION(BlueprintSetter)
+	void SetReplicatePlayback(bool ReplicatePlayback);
+
+#if WITH_EDITOR
+	virtual bool GetReferencedContentObjects(TArray<UObject*>& Objects) const override;
+#endif //WITH_EDITOR
+
 protected:
+	//~ Begin UObject interface
+	virtual bool ReplicateSubobjects(UActorChannel* Channel, FOutBunch* Bunch, FReplicationFlags *RepFlags) override;
+	virtual void PostInitProperties() override;
+	virtual void PostLoad() override;
+	//~ End UObject interface
+	
 	//~ Begin AActor interface
 	virtual void PostInitializeComponents() override;
 	virtual void BeginPlay() override;
 	virtual void EndPlay(const EEndPlayReason::Type EndPlayReason) override;
 	virtual void RewindForReplay() override;
+#if WITH_EDITOR
+	virtual bool CanChangeIsSpatiallyLoadedFlag() const override { return false; }
+#endif
 	//~ End AActor interface
 
-	/** Customized TimeOfDay initialization of the SequencePlayer */
-	void InitializeSequencePlayer();
+	//~ Begin IMovieSceneSequenceActor interface
+	virtual void TickFromSequenceTickManager(float DeltaSeconds) override;
+	//~ End IMovieSceneSequenceActor interface
+
+	//~ Begin IMovieScenePlaybackClient interface
+	virtual bool RetrieveBindingOverrides(const FGuid& InBindingId, FMovieSceneSequenceID InSequenceID, TArray<UObject*, TInlineAllocator<1>>& OutObjects) const override;
+	virtual UObject* GetInstanceData() const override;
+	virtual bool GetIsReplicatedPlayback() const override;
+	//~ End IMovieScenePlaybackClient interface
+
+	//~ Begin IMovieSceneBindingOwnerInterface
+#if WITH_EDITOR
+	virtual TSharedPtr<FStructOnScope> GetObjectPickerProxy(TSharedPtr<IPropertyHandle> PropertyHandle) override;
+	virtual void UpdateObjectFromProxy(FStructOnScope& Proxy, IPropertyHandle& ObjectPropertyHandle) override;
+	virtual UMovieSceneSequence* RetrieveOwnedSequence() const override
+	{
+		return DaySequenceAsset;
+	}
+#endif
+	//~ End IMovieSceneBindingOwnerInterface
+
+	/** Initialize SequencePlayer with transient master sequence */
+	void InitializePlayer();
+
+	void UpdateDaySequence(ULevelSequence* SequenceAsset);
 
 	/** Compute PlaybackSettings from day cycle properties */
 	FMovieSceneSequencePlaybackSettings GetPlaybackSettings(const ULevelSequence* Sequence) const;
 
+public:
+	UPROPERTY(Instanced, transient, replicated, BlueprintReadOnly, BlueprintGetter=GetSequencePlayer, Category="Playback", meta=(ExposeFunctionCategories="Sequencer|Player"))
+	TObjectPtr<ULevelSequencePlayer> SequencePlayer;
+
+	UPROPERTY(EditAnywhere, BlueprintReadOnly, Category="General", meta=(AllowedClasses="LevelSequence"))
+	TObjectPtr<ULevelSequence> DaySequenceAsset;
+
+	/** Mapping of actors to override the sequence bindings with */
+	UPROPERTY(Instanced, BlueprintReadOnly, Category="General")
+	TObjectPtr<UMovieSceneBindingOverrides> BindingOverrides;
+	
+	/** If true, playback of this level sequence on the server will be synchronized across other clients */
+	UPROPERTY(EditAnywhere, DisplayName="Replicate Playback", BlueprintReadWrite, BlueprintSetter=SetReplicatePlayback, Category=Replication)
+	uint8 bReplicatePlayback : 1;
+
 protected:
-	UPROPERTY()
+	UPROPERTY(Transient)
 	TObjectPtr<ULevelSequence> MasterSequence;
 
 #if WITH_EDITORONLY_DATA
@@ -59,3 +140,4 @@ protected:
 	FTimecode InitialTimeOfDay;
 };
 
+
diff --git a/Engine/Plugins/Lighting/TimeOfDay/Source/TimeOfDayEditor/Private/TimeOfDayActorDetails.cpp b/Engine/Plugins/Lighting/TimeOfDay/Source/TimeOfDayEditor/Private/TimeOfDayActorDetails.cpp
--- a/Engine/Plugins/Lighting/TimeOfDay/Source/TimeOfDayEditor/Private/TimeOfDayActorDetails.cpp
+++ b/Engine/Plugins/Lighting/TimeOfDay/Source/TimeOfDayEditor/Private/TimeOfDayActorDetails.cpp
@@ -1,8 +1,14 @@
 // Copyright Epic Games, Inc. All Rights Reserved.
 #include "TimeOfDayActorDetails.h"
+#include "TimeOfDayActor.h"
+#include "Widgets/SBoxPanel.h"
+#include "Widgets/Input/SButton.h"
+#include "Editor.h"
 #include "DetailLayoutBuilder.h"
+#include "DetailCategoryBuilder.h"
 #include "DetailWidgetRow.h"
-#include "LevelSequenceActor.h"
+#include "Algo/Transform.h"
+#include "Subsystems/AssetEditorSubsystem.h"
 
 #define LOCTEXT_NAMESPACE "TimeOfDayActorDetails"
 
@@ -11,24 +17,113 @@ TSharedRef<IDetailCustomization> FTimeOfDayActorDetails::MakeInstance()
 	return MakeShared<FTimeOfDayActorDetails>();
 }
 
+void AddAllSubObjectProperties(TArray<UObject*>& SubObjects, IDetailCategoryBuilder& Category, TAttribute<EVisibility> Visibility = TAttribute<EVisibility>(EVisibility::Visible))
+{
+	SubObjects.Remove(nullptr);
+	if (!SubObjects.Num())
+	{
+		return;
+	}
+
+	for (const FProperty* TestProperty : TFieldRange<FProperty>(SubObjects[0]->GetClass()))
+	{
+		if (TestProperty->HasAnyPropertyFlags(CPF_Edit))
+		{
+			const bool bAdvancedDisplay = TestProperty->HasAnyPropertyFlags(CPF_AdvancedDisplay);
+			const EPropertyLocation::Type PropertyLocation = bAdvancedDisplay ? EPropertyLocation::Advanced : EPropertyLocation::Common;
+
+			IDetailPropertyRow* NewRow = Category.AddExternalObjectProperty(SubObjects, TestProperty->GetFName(), PropertyLocation);
+			if (NewRow)
+			{
+				NewRow->Visibility(Visibility);
+			}
+		}
+	}
+}
+
 void FTimeOfDayActorDetails::CustomizeDetails(IDetailLayoutBuilder& DetailLayout)
 {
-	DetailLayout.HideProperty("DefaultComponents");
+	const TArray<TWeakObjectPtr<UObject>>& SelectedObjects = DetailLayout.GetSelectedObjects();
+	for (int32 ObjectIndex = 0; ObjectIndex < SelectedObjects.Num(); ++ObjectIndex)
+	{
+		const TWeakObjectPtr<UObject>& CurrentObject = SelectedObjects[ObjectIndex];
+		if (CurrentObject.IsValid())
+		{
+			if (ATimeOfDayActor* CurrentTimeOfDayActor = Cast<ATimeOfDayActor>(CurrentObject.Get()))
+			{
+				TimeOfDayActor = CurrentTimeOfDayActor;
+				break;
+			}
+		}
+	}
+
+	TArray<ATimeOfDayActor*> TimeOfDayActors;
+	{
+		TArray<TWeakObjectPtr<>> ObjectPtrs;
+		DetailLayout.GetObjectsBeingCustomized(ObjectPtrs);
+
+		for (TWeakObjectPtr<> WeakObj : ObjectPtrs)
+		{
+			if (ATimeOfDayActor* Actor = Cast<ATimeOfDayActor>(WeakObj.Get()))
+			{
+				TimeOfDayActors.Add(Actor);
+			}
+		}
+	}
 	
-	// Hide level level sequence actor stuff
-	DetailLayout.HideCategory("Playback");
-	DetailLayout.HideCategory("BurnInOptions");
-	DetailLayout.HideCategory("Aspect Ratio");
-	DetailLayout.HideCategory("Cinematic");
-	DetailLayout.HideCategory("InstanceData");
-	DetailLayout.HideCategory("Tick");
-
-	const TSharedPtr<IPropertyHandle> LevelSequenceProp = DetailLayout.GetProperty(GET_MEMBER_NAME_CHECKED(ALevelSequenceActor, LevelSequenceAsset), ALevelSequenceActor::StaticClass());
-	if (IDetailPropertyRow* DetailRow = DetailLayout.EditDefaultProperty(LevelSequenceProp))
+	DetailLayout.HideProperty("DefaultComponents");
+
+	IDetailCategoryBuilder& GeneralCategory = DetailLayout.EditCategory( "General", NSLOCTEXT("GeneralDetails", "General", "General"), ECategoryPriority::Important );
+
+	GeneralCategory.AddCustomRow( NSLOCTEXT("TimeOfDayActorDetails", "OpenLevelSequence", "Open Level Sequence") )
+	.RowTag("OpenLevelSequence")
+	[
+		SNew(SHorizontalBox)
+		+SHorizontalBox::Slot()
+		.FillWidth(1.f)
+		.Padding(0, 5, 10, 5)
+		[
+			SNew(SButton)
+			.ContentPadding(3)
+			.VAlign(VAlign_Center)
+			.HAlign(HAlign_Center)
+			.IsEnabled( this, &FTimeOfDayActorDetails::CanOpenLevelSequenceForActor )
+			.OnClicked( this, &FTimeOfDayActorDetails::OnOpenLevelSequenceForActor )
+			.Text( NSLOCTEXT("TimeOfDayActorDetails", "OpenLevelSequence", "Open Level Sequence") )
+		]
+	];
+
+	TArray<UObject*> SubObjects;
+
+	IDetailCategoryBuilder& BindingOverridesCategory = DetailLayout.EditCategory( "BindingOverrides", LOCTEXT("BindingOverrides", "Binding Overrides") ).InitiallyCollapsed(false);
+	{
+		SubObjects.Reset();
+		Algo::Transform(TimeOfDayActors, SubObjects, &ATimeOfDayActor::BindingOverrides);
+		AddAllSubObjectProperties(SubObjects, BindingOverridesCategory);
+	}
+}
+
+bool FTimeOfDayActorDetails::CanOpenLevelSequenceForActor() const
+{
+	if (TimeOfDayActor.IsValid())
+	{
+		return TimeOfDayActor.Get()->GetDaySequence() != nullptr;
+	}
+	return false;
+}
+
+FReply FTimeOfDayActorDetails::OnOpenLevelSequenceForActor() const
+{
+	if (TimeOfDayActor.IsValid())
 	{
-		DetailRow->DisplayName(LOCTEXT("EnvironmentPreset", "Environment Preset"));
+		if (UObject* LoadedObject = TimeOfDayActor.Get()->GetDaySequence())
+		{
+			GEditor->GetEditorSubsystem<UAssetEditorSubsystem>()->OpenEditorForAsset(LoadedObject);
+		}
 	}
+	return FReply::Handled();
 }
 
 
+
 #undef LOCTEXT_NAMESPACE
diff --git a/Engine/Plugins/Lighting/TimeOfDay/Source/TimeOfDayEditor/Private/TimeOfDayActorDetails.h b/Engine/Plugins/Lighting/TimeOfDay/Source/TimeOfDayEditor/Private/TimeOfDayActorDetails.h
--- a/Engine/Plugins/Lighting/TimeOfDay/Source/TimeOfDayEditor/Private/TimeOfDayActorDetails.h
+++ b/Engine/Plugins/Lighting/TimeOfDay/Source/TimeOfDayEditor/Private/TimeOfDayActorDetails.h
@@ -3,17 +3,27 @@
 #pragma once
 
 #include "CoreMinimal.h"
+#include "UObject/WeakObjectPtr.h"
+#include "Input/Reply.h"
 #include "IDetailCustomization.h"
 
 class IDetailLayoutBuilder;
+class ATimeOfDayActor;
 
 class FTimeOfDayActorDetails : public IDetailCustomization
 {
 public:
 	/** Makes a new instance of this detail layout class for a specific detail view requesting it */
 	static TSharedRef<IDetailCustomization> MakeInstance();
+	
 private:
 	/** IDetailCustomization interface */
 	virtual void CustomizeDetails(IDetailLayoutBuilder& DetailLayout) override;
+
+	bool CanOpenLevelSequenceForActor() const;
+	FReply OnOpenLevelSequenceForActor() const;
+	
 private:
+	/** The selected level sequence actor */
+	TWeakObjectPtr<ATimeOfDayActor> TimeOfDayActor;
 };
diff --git a/Engine/Plugins/Lighting/TimeOfDay/Source/TimeOfDayEditor/Private/TimeOfDayActorFactory.cpp b/Engine/Plugins/Lighting/TimeOfDay/Source/TimeOfDayEditor/Private/TimeOfDayActorFactory.cpp
new file mode 100644
--- /dev/null
+++ b/Engine/Plugins/Lighting/TimeOfDay/Source/TimeOfDayEditor/Private/TimeOfDayActorFactory.cpp
@@ -0,0 +1,55 @@
+// Copyright Epic Games, Inc. All Rights Reserved.
+
+#include "TimeOfDayActorFactory.h"
+#include "TimeOfDayActor.h"
+#include "LevelSequence.h"
+
+#define LOCTEXT_NAMESPACE "TimeOfDayEditor"
+
+UTimeOfDayActorFactory::UTimeOfDayActorFactory(const FObjectInitializer& ObjectInitializer)
+	: Super(ObjectInitializer)
+{
+	DisplayName = LOCTEXT("TimeOfDayDisplayName", "Time of Day");
+	NewActorClass = ATimeOfDayActor::StaticClass();
+}
+
+bool UTimeOfDayActorFactory::CanCreateActorFrom( const FAssetData& AssetData, FText& OutErrorMsg )
+{
+	if (UActorFactory::CanCreateActorFrom(AssetData, OutErrorMsg))
+	{
+		return true;
+	}
+
+	if (AssetData.IsValid() && !AssetData.IsInstanceOf(ULevelSequence::StaticClass()))
+	{
+		OutErrorMsg = NSLOCTEXT("CanCreateActor", "NoLevelSequenceAsset", "A valid sequencer asset must be specified.");
+		return false;
+	}
+	
+	return true;
+}
+
+AActor* UTimeOfDayActorFactory::SpawnActor(UObject* InAsset, ULevel* InLevel, const FTransform& InTransform, const FActorSpawnParameters& InSpawnParams)
+{
+	ATimeOfDayActor* NewActor = Cast<ATimeOfDayActor>(Super::SpawnActor(InAsset, InLevel, InTransform, InSpawnParams));
+	if (NewActor)
+	{
+		if (ULevelSequence* LevelSequence = Cast<ULevelSequence>(InAsset))
+		{
+			NewActor->SetDaySequence(LevelSequence);
+		}
+	}
+	return NewActor;
+}
+
+UObject* UTimeOfDayActorFactory::GetAssetFromActorInstance(AActor* Instance)
+{
+	if (const ATimeOfDayActor* TimeOfDayActor = Cast<ATimeOfDayActor>(Instance))
+	{
+		return TimeOfDayActor->GetDaySequence();
+	}
+	return nullptr;
+}
+
+#undef LOCTEXT_NAMESPACE
+
diff --git a/Engine/Plugins/Lighting/TimeOfDay/Source/TimeOfDayEditor/Public/TimeOfDayActorFactory.h b/Engine/Plugins/Lighting/TimeOfDay/Source/TimeOfDayEditor/Public/TimeOfDayActorFactory.h
new file mode 100644
--- /dev/null
+++ b/Engine/Plugins/Lighting/TimeOfDay/Source/TimeOfDayEditor/Public/TimeOfDayActorFactory.h
@@ -0,0 +1,23 @@
+// Copyright Epic Games, Inc. All Rights Reserved.
+
+#pragma once
+
+#include "CoreMinimal.h"
+#include "ActorFactories/ActorFactory.h"
+#include "TimeOfDayActorFactory.generated.h"
+
+class AActor;
+struct FAssetData;
+
+UCLASS()
+class UTimeOfDayActorFactory : public UActorFactory
+{
+	GENERATED_UCLASS_BODY()
+
+	//~ Begin UActorFactory Interface
+	virtual AActor* SpawnActor(UObject* InAsset, ULevel* InLevel, const FTransform& InTransform, const FActorSpawnParameters& InSpawnParams) override;
+	virtual bool CanCreateActorFrom( const FAssetData& AssetData, FText& OutErrorMsg ) override;
+	virtual UObject* GetAssetFromActorInstance(AActor* ActorInstance) override;
+	//~ End UActorFactory Interface
+};
+
