TimeOfDay: Make TimeOfDayActorFactory.h private.

diff --git a/Engine/Plugins/Lighting/TimeOfDay/Source/TimeOfDayEditor/Private/TimeOfDayActorFactory.h b/Engine/Plugins/Lighting/TimeOfDay/Source/TimeOfDayEditor/Private/TimeOfDayActorFactory.h
new file mode 100644
--- /dev/null
+++ b/Engine/Plugins/Lighting/TimeOfDay/Source/TimeOfDayEditor/Private/TimeOfDayActorFactory.h
@@ -0,0 +1,23 @@
+// Copyright Epic Games, Inc. All Rights Reserved.
+
+#pragma once
+
+#include "CoreMinimal.h"
+#include "ActorFactories/ActorFactory.h"
+#include "TimeOfDayActorFactory.generated.h"
+
+class AActor;
+struct FAssetData;
+
+UCLASS()
+class UTimeOfDayActorFactory : public UActorFactory
+{
+	GENERATED_UCLASS_BODY()
+
+	//~ Begin UActorFactory Interface
+	virtual AActor* SpawnActor(UObject* InAsset, ULevel* InLevel, const FTransform& InTransform, const FActorSpawnParameters& InSpawnParams) override;
+	virtual bool CanCreateActorFrom( const FAssetData& AssetData, FText& OutErrorMsg ) override;
+	virtual UObject* GetAssetFromActorInstance(AActor* ActorInstance) override;
+	//~ End UActorFactory Interface
+};
+
diff --git a/Engine/Plugins/Lighting/TimeOfDay/Source/TimeOfDayEditor/Public/TimeOfDayActorFactory.h b/Engine/Plugins/Lighting/TimeOfDay/Source/TimeOfDayEditor/Public/TimeOfDayActorFactory.h
deleted file mode 100644
